Remove unneeded global lock on nvtt object construction

#tests Tested this CL locally for months on 64 cores CPU and never had any issue.

diff --git a/Engine/Source/Developer/TextureFormatDXT/Private/TextureFormatDXT.cpp b/Engine/Source/Developer/TextureFormatDXT/Private/TextureFormatDXT.cpp
--- a/Engine/Source/Developer/TextureFormatDXT/Private/TextureFormatDXT.cpp
+++ b/Engine/Source/Developer/TextureFormatDXT/Private/TextureFormatDXT.cpp
@@ -98,9 +98,6 @@ struct FNVErrorHandler : public nvtt::ErrorHandler
 	bool bSuccess;
 };
 
-/** Critical section to isolate construction of nvtt objects */
-FCriticalSection GNVCompressionCriticalSection;
-
 /**
  * All state objects needed for NVTT.
  */
@@ -328,7 +325,6 @@ static bool CompressImageUsingNVTT(
 	{
 		FNVTTCompressor* Compressor = NULL;
 		{
-			FScopeLock ScopeLock(&GNVCompressionCriticalSection);
 			Compressor = new FNVTTCompressor(
 				SourceData,
 				PixelFormat,
@@ -343,7 +339,6 @@ static bool CompressImageUsingNVTT(
 		}
 		bool bSuccess = Compressor->Compress();
 		{
-			FScopeLock ScopeLock(&GNVCompressionCriticalSection);
 			delete Compressor;
 			Compressor = NULL;
 		}
@@ -357,7 +352,6 @@ static bool CompressImageUsingNVTT(
 	TIndirectArray<FNVTTCompressor> Compressors;
 	Compressors.Empty(NumBatches);
 	{
-		FScopeLock ScopeLock(&GNVCompressionCriticalSection);
 		const uint8* Src = (const uint8*)SourceData;
 		uint8* Dest = OutCompressedData.GetData();
 		for (int32 BatchIndex = 0; BatchIndex < NumBatches; ++BatchIndex)
@@ -401,7 +395,6 @@ static bool CompressImageUsingNVTT(
 
 	// Release compressors
 	{
-		FScopeLock ScopeLock(&GNVCompressionCriticalSection);
 		Compressors.Empty();
 	}
 
