TextureFormatDXT crash fix

This moves the NVTT DLL loading code that was removed from TextureCompressorModule (where it doesn't belong) to here

diff --git a/Engine/Source/Developer/TextureFormatDXT/Private/TextureFormatDXT.cpp b/Engine/Source/Developer/TextureFormatDXT/Private/TextureFormatDXT.cpp
--- a/Engine/Source/Developer/TextureFormatDXT/Private/TextureFormatDXT.cpp
+++ b/Engine/Source/Developer/TextureFormatDXT/Private/TextureFormatDXT.cpp
@@ -17,6 +17,7 @@
 #include "TextureBuildFunction.h"
 #include "DerivedDataBuildFunctionFactory.h"
 #include "DerivedDataSharedString.h"
+#include "Misc/Paths.h"
 
 THIRD_PARTY_INCLUDES_START
 	#include "nvtt/nvtt.h"
@@ -573,7 +574,40 @@ public:
 		return Singleton;
 	}
 
+	// IModuleInterface implementation.
+	virtual void StartupModule() override
+	{
+#if PLATFORM_WINDOWS
+#if PLATFORM_64BITS
+		if (FWindowsPlatformMisc::HasAVX2InstructionSupport())
+		{
+			nvTextureToolsHandle = FPlatformProcess::GetDllHandle(*(FPaths::EngineDir() / TEXT("Binaries/ThirdParty/nvTextureTools/Win64/AVX2/nvtt_64.dll")));
+		}
+		else
+		{
+			nvTextureToolsHandle = FPlatformProcess::GetDllHandle(*(FPaths::EngineDir() / TEXT("Binaries/ThirdParty/nvTextureTools/Win64/nvtt_64.dll")));
+		}
+#else	//32-bit platform
+		nvTextureToolsHandle = FPlatformProcess::GetDllHandle(*(FPaths::EngineDir() / TEXT("Binaries/ThirdParty/nvTextureTools/Win32/nvtt_.dll")));
+#endif
+#endif	//PLATFORM_WINDOWS
+	}
+
+	virtual void ShutdownModule() override
+	{
+#if PLATFORM_WINDOWS
+		FPlatformProcess::FreeDllHandle(nvTextureToolsHandle);
+		nvTextureToolsHandle = 0;
+#endif
+	}
+
 	static inline UE::DerivedData::TBuildFunctionFactory<FDXTTextureBuildFunction> BuildFunctionFactory;
+
+private:
+#if PLATFORM_WINDOWS
+	// Handle to the nvtt dll
+	void* nvTextureToolsHandle;
+#endif	//PLATFORM_WINDOWS
 };
 
 IMPLEMENT_MODULE(FTextureFormatDXTModule, TextureFormatDXT);
