Fix for cook hang when TextureFormatOodle is disabled for DXT format.

diff --git a/Engine/Source/Developer/TextureFormatDXT/Private/TextureFormatDXT.cpp b/Engine/Source/Developer/TextureFormatDXT/Private/TextureFormatDXT.cpp
--- a/Engine/Source/Developer/TextureFormatDXT/Private/TextureFormatDXT.cpp
+++ b/Engine/Source/Developer/TextureFormatDXT/Private/TextureFormatDXT.cpp
@@ -578,7 +578,7 @@
 		}
 
 		// Lock so only one thread does init :
-		nvTextureToolsHandleLock.Lock();
+		FScopeLock HandleLock(&nvTextureToolsHandleLock);
 		
 		// double check inside lock :
 		if ( nvTextureToolsHandle != nullptr )
@@ -594,8 +594,6 @@
 		{
 			nvTextureToolsHandle = FPlatformProcess::GetDllHandle(*(FPaths::EngineDir() / TEXT("Binaries/ThirdParty/nvTextureTools/Win64/nvtt_64.dll")));
 		}
-
-		nvTextureToolsHandleLock.Unlock();
 #endif	//PLATFORM_WINDOWS
 	}
 
