add missing arm64 runtime library for nvTexureTools

#rnx

diff --git a/Engine/Source/Developer/TextureFormatDXT/Private/TextureFormatDXT.cpp b/Engine/Source/Developer/TextureFormatDXT/Private/TextureFormatDXT.cpp
--- a/Engine/Source/Developer/TextureFormatDXT/Private/TextureFormatDXT.cpp
+++ b/Engine/Source/Developer/TextureFormatDXT/Private/TextureFormatDXT.cpp
@@ -583,7 +583,9 @@ public:
 		{
 			return;
 		}
-
+#if PLATFORM_CPU_ARM_FAMILY
+		nvTextureToolsHandle = FPlatformProcess::GetDllHandle(*(FPaths::EngineDir() / TEXT("Binaries/ThirdParty/nvTextureTools/Win64/arm64/nvtt_64.dll")));
+#else
 		if (FWindowsPlatformMisc::HasAVX2InstructionSupport())
 		{
 			nvTextureToolsHandle = FPlatformProcess::GetDllHandle(*(FPaths::EngineDir() / TEXT("Binaries/ThirdParty/nvTextureTools/Win64/AVX2/nvtt_64.dll")));
@@ -592,6 +594,7 @@ public:
 		{
 			nvTextureToolsHandle = FPlatformProcess::GetDllHandle(*(FPaths::EngineDir() / TEXT("Binaries/ThirdParty/nvTextureTools/Win64/nvtt_64.dll")));
 		}
+#endif  //PLATFORM_CPU_ARM_FAMILY
 #endif	//PLATFORM_WINDOWS
 	}
 
